Fix incorrect module name, causing monolithic compile error.

diff --git a/Engine/Plugins/Runtime/SoundUtilities/Source/SoundUtilitiesEditor/Private/SoundUtilitiesEditorModule.cpp b/Engine/Plugins/Runtime/SoundUtilities/Source/SoundUtilitiesEditor/Private/SoundUtilitiesEditorModule.cpp
--- a/Engine/Plugins/Runtime/SoundUtilities/Source/SoundUtilitiesEditor/Private/SoundUtilitiesEditorModule.cpp
+++ b/Engine/Plugins/Runtime/SoundUtilities/Source/SoundUtilitiesEditor/Private/SoundUtilitiesEditorModule.cpp
@@ -11,7 +11,7 @@
 #include "AudioEditorModule.h"
 
 
-IMPLEMENT_MODULE(FSoundUtilitiesEditorModule, FSynthesisEditor)
+IMPLEMENT_MODULE(FSoundUtilitiesEditorModule, SoundUtilitiesEditor)
 
 void FSoundUtilitiesEditorModule::StartupModule()
 {
