Removing references to ue4

diff --git a/Engine/Plugins/Runtime/SoundUtilities/SoundUtilities.uplugin b/Engine/Plugins/Runtime/SoundUtilities/SoundUtilities.uplugin
--- a/Engine/Plugins/Runtime/SoundUtilities/SoundUtilities.uplugin
+++ b/Engine/Plugins/Runtime/SoundUtilities/SoundUtilities.uplugin
@@ -3,7 +3,7 @@
 	"Version" : 1,
 	"VersionName" : "1.0",
 	"FriendlyName" : "Sound Utilities",
-	"Description" : "A variety of BP functions, objects, and utilities for audio with UE4.",
+	"Description" : "A variety of BP functions, objects, and utilities for audio.",
 	"Category" : "Audio",
 	"CreatedBy" : "Epic Games, Inc.",
 	"CreatedByURL" : "http://epicgames.com",
