Fixing simple sound variation add

diff --git a/Engine/Plugins/Runtime/SoundUtilities/Source/SoundUtilities/Private/SoundSimple.cpp b/Engine/Plugins/Runtime/SoundUtilities/Source/SoundUtilities/Private/SoundSimple.cpp
--- a/Engine/Plugins/Runtime/SoundUtilities/Source/SoundUtilities/Private/SoundSimple.cpp
+++ b/Engine/Plugins/Runtime/SoundUtilities/Source/SoundUtilities/Private/SoundSimple.cpp
@@ -6,22 +6,24 @@
 void USoundSimple::CacheValues()
 {
 	MaxDistance = 0.0f;
-	for (int32 i = 0; i < Variations.Num(); ++i)
-	{
-		float SoundWaveMaxAudibleDistance = Variations[i].SoundWave->GetMaxDistance();
-		if (SoundWaveMaxAudibleDistance > MaxDistance)
-		{
-			MaxDistance = SoundWaveMaxAudibleDistance;
-		}
-	}
-
 	Duration = 0.0f;
+
 	for (int32 i = 0; i < Variations.Num(); ++i)
 	{
-		float SoundWaveMaxDuration = Variations[i].SoundWave->GetDuration();
-		if (SoundWaveMaxDuration > Duration)
+		USoundWave* SoundWaveVariation = Variations[i].SoundWave;
+		if (SoundWaveVariation)
 		{
-			Duration = SoundWaveMaxDuration;
+			float SoundWaveMaxAudibleDistance = SoundWaveVariation->GetMaxDistance();
+			if (SoundWaveMaxAudibleDistance > MaxDistance)
+			{
+				MaxDistance = SoundWaveMaxAudibleDistance;
+			}
+
+			float SoundWaveMaxDuration = SoundWaveVariation->GetDuration();
+			if (SoundWaveMaxDuration > Duration)
+			{
+				Duration = SoundWaveMaxDuration;
+			}
 		}
 	}
 }
@@ -67,7 +69,10 @@ void USoundSimple::ChooseSoundWave()
 	float ProbabilitySum = 0.0f;
 	for (int32 i = 0; i < Variations.Num(); ++i)
 	{
-		ProbabilitySum += Variations[i].ProbabilityWeight;
+		if (Variations[i].SoundWave)
+		{
+			ProbabilitySum += Variations[i].ProbabilityWeight;
+		}
 	}
 
 	float Choice = FMath::FRandRange(0.0f, ProbabilitySum);
@@ -77,20 +82,25 @@ void USoundSimple::ChooseSoundWave()
 	int32 ChosenIndex = 0;
 	for (int32 i = 0; i < Variations.Num(); ++i)
 	{
-		float NextSum = ProbabilitySum + Variations[i].ProbabilityWeight;
-
-		if (Choice >= ProbabilitySum && Choice < NextSum)
+		if (Variations[i].SoundWave)
 		{
-			ChosenIndex = i;
-			break;
-		}
+			float NextSum = ProbabilitySum + Variations[i].ProbabilityWeight;
 
-		ProbabilitySum = NextSum;
+			if (Choice >= ProbabilitySum && Choice < NextSum)
+			{
+				ChosenIndex = i;
+				break;
+			}
+
+			ProbabilitySum = NextSum;
+		}
 	}
 	
 	check(ChosenIndex < Variations.Num());
 	FSoundVariation& SoundVariation = Variations[ChosenIndex];
 
+	check(SoundVariation.SoundWave);
+
 	// Now choise the volume and pitch to use based on prob ranges
 	float Volume = FMath::FRandRange(SoundVariation.VolumeRange[0], SoundVariation.VolumeRange[1]);
 	float Pitch = FMath::FRandRange(SoundVariation.PitchRange[0], SoundVariation.PitchRange[1]);
