Editor - Followup to CL 22508580, fixing a new deprecated instance of someone using SelectedObjects.

diff --git a/Engine/Plugins/Runtime/SoundUtilities/Source/SoundUtilitiesEditor/Private/SoundWaveAssetActionExtender.cpp b/Engine/Plugins/Runtime/SoundUtilities/Source/SoundUtilitiesEditor/Private/SoundWaveAssetActionExtender.cpp
--- a/Engine/Plugins/Runtime/SoundUtilities/Source/SoundUtilitiesEditor/Private/SoundWaveAssetActionExtender.cpp
+++ b/Engine/Plugins/Runtime/SoundUtilities/Source/SoundUtilitiesEditor/Private/SoundWaveAssetActionExtender.cpp
@@ -28,61 +28,47 @@ void FSoundWaveAssetActionExtender::RegisterMenus()
 	
 	Section.AddDynamicEntry("SoundWaveAsset", FNewToolMenuSectionDelegate::CreateLambda([](FToolMenuSection& InSection)
 	{
-		UContentBrowserAssetContextMenuContext* Context = InSection.FindContext<UContentBrowserAssetContextMenuContext>();
-		if (!Context || Context->SelectedObjects.IsEmpty())
+		if (UContentBrowserAssetContextMenuContext* Context = InSection.FindContext<UContentBrowserAssetContextMenuContext>())
 		{
-			return;
-		}
-
-		for (const TWeakObjectPtr<UObject>& Object : Context->SelectedObjects)
-		{
-			if (Object->IsA<USoundWaveProcedural>())
+			if (Algo::AnyOf(Context->SelectedAssets, [](const FAssetData& AssetData){ return AssetData.IsInstanceOf<USoundWaveProcedural>(); }))
 			{
 				return;
 			}
-		}
-		const TAttribute<FText> Label = LOCTEXT("SoundWave_CreateSimpleSound", "Create Simple Sound");
-		const TAttribute<FText> ToolTip = LOCTEXT("SoundWave_CreateSimpleSoundTooltip", "Creates a simple sound asset using the selected sound waves.");
-		const FSlateIcon Icon = FSlateIcon(FAppStyle::GetAppStyleSetName(), "ClassIcon.SoundSimple");
-		const FToolMenuExecuteAction UIAction = FToolMenuExecuteAction::CreateStatic(&FSoundWaveAssetActionExtender::ExecuteCreateSimpleSound);
 
-		InSection.AddMenuEntry("SoundWave_CreateSimpleSound", Label, ToolTip, Icon, UIAction);
+			const TAttribute<FText> Label = LOCTEXT("SoundWave_CreateSimpleSound", "Create Simple Sound");
+			const TAttribute<FText> ToolTip = LOCTEXT("SoundWave_CreateSimpleSoundTooltip", "Creates a simple sound asset using the selected sound waves.");
+			const FSlateIcon Icon = FSlateIcon(FAppStyle::GetAppStyleSetName(), "ClassIcon.SoundSimple");
+			const FToolMenuExecuteAction UIAction = FToolMenuExecuteAction::CreateStatic(&FSoundWaveAssetActionExtender::ExecuteCreateSimpleSound);
+
+			InSection.AddMenuEntry("SoundWave_CreateSimpleSound", Label, ToolTip, Icon, UIAction);
+		}
 	}));
 }
 
 void FSoundWaveAssetActionExtender::ExecuteCreateSimpleSound(const FToolMenuContext& MenuContext)
 {
-	UContentBrowserAssetContextMenuContext* Context = MenuContext.FindContext<UContentBrowserAssetContextMenuContext>();
-	if (!Context || Context->SelectedObjects.Num() == 0)
+	if (const UContentBrowserAssetContextMenuContext* Context = UContentBrowserAssetContextMenuContext::FindContextWithAssets(MenuContext))
 	{
-		return;
-	}
-
-	const FString DefaultSuffix = TEXT("_SimpleSound");
+		const FString DefaultSuffix = TEXT("_SimpleSound");
 	
-	if (USoundWave* SoundWave = Cast<USoundWave>(Context->SelectedObjects[0]))
-	{
-		// Determine an appropriate name
-		FString Name;
-		FString PackagePath;
+		TArray<USoundWave*> SoundWaves = Context->LoadSelectedObjectsIf<USoundWave>([](const FAssetData& AssetData) { return !AssetData.IsInstanceOf<USoundWaveProcedural>(); });
 
-		FAssetToolsModule& AssetToolsModule = FModuleManager::Get().LoadModuleChecked<FAssetToolsModule>("AssetTools");
-		AssetToolsModule.Get().CreateUniqueAssetName(SoundWave->GetOutermost()->GetName(), DefaultSuffix, PackagePath, Name);
-
-		// Create the factory used to generate the asset
-		USoundSimpleFactory* Factory = NewObject<USoundSimpleFactory>();
-		Factory->SoundWaves.Reset();
-		for (const TWeakObjectPtr<UObject>& Object : Context->SelectedObjects)
+		if (!SoundWaves.IsEmpty())
 		{
-			if (USoundWave* Wave = Cast<USoundWave>(Object))
-			{
-				Factory->SoundWaves.Add(Wave);
-			}
-		}
+			// Determine an appropriate name
+			FString Name;
+			FString PackagePath;
 
-		FContentBrowserModule& ContentBrowserModule = FModuleManager::LoadModuleChecked<FContentBrowserModule>("ContentBrowser");
-		ContentBrowserModule.Get().CreateNewAsset(Name, FPackageName::GetLongPackagePath(PackagePath), USoundSimple::StaticClass(), Factory);
+			FAssetToolsModule& AssetToolsModule = FModuleManager::Get().LoadModuleChecked<FAssetToolsModule>("AssetTools");
+			AssetToolsModule.Get().CreateUniqueAssetName(SoundWaves[0]->GetOutermost()->GetName(), DefaultSuffix, PackagePath, Name);
 
+			// Create the factory used to generate the asset
+			USoundSimpleFactory* Factory = NewObject<USoundSimpleFactory>();
+			Factory->SoundWaves = SoundWaves;
+
+			FContentBrowserModule& ContentBrowserModule = FModuleManager::LoadModuleChecked<FContentBrowserModule>("ContentBrowser");
+			ContentBrowserModule.Get().CreateNewAsset(Name, FPackageName::GetLongPackagePath(PackagePath), USoundSimple::StaticClass(), Factory);
+		}
 	}
 }
 
