Non-unity fix

#rnx

diff --git a/Engine/Plugins/Runtime/SoundUtilities/Source/SoundUtilitiesEditor/Private/SoundWaveAssetActionExtender.cpp b/Engine/Plugins/Runtime/SoundUtilities/Source/SoundUtilitiesEditor/Private/SoundWaveAssetActionExtender.cpp
--- a/Engine/Plugins/Runtime/SoundUtilities/Source/SoundUtilitiesEditor/Private/SoundWaveAssetActionExtender.cpp
+++ b/Engine/Plugins/Runtime/SoundUtilities/Source/SoundUtilitiesEditor/Private/SoundWaveAssetActionExtender.cpp
@@ -12,6 +12,7 @@
 #include "Sound/SoundWaveProcedural.h"
 #include "SoundSimple.h"
 #include "SoundSimpleFactory.h"
+#include "Algo/AnyOf.h"
 
 #define LOCTEXT_NAMESPACE "AssetTypeActions"
 
