Fixing crash in simple sound

[RN] Fixing crash in simple sound when there are no variations.

diff --git a/Engine/Plugins/Runtime/SoundUtilities/Source/SoundUtilities/Private/SoundSimple.cpp b/Engine/Plugins/Runtime/SoundUtilities/Source/SoundUtilities/Private/SoundSimple.cpp
--- a/Engine/Plugins/Runtime/SoundUtilities/Source/SoundUtilities/Private/SoundSimple.cpp
+++ b/Engine/Plugins/Runtime/SoundUtilities/Source/SoundUtilities/Private/SoundSimple.cpp
@@ -60,13 +60,22 @@ void USoundSimple::Parse(class FAudioDevice* AudioDevice, const UPTRINT NodeWave
 		ChooseSoundWave();
 	}
 
+	if (!SoundWave)
+	{
+		return;
+	}
+
 	// Forward the parse to the chosen sound wave
-	check(SoundWave);
 	SoundWave->Parse(AudioDevice, NodeWaveInstanceHash, ActiveSound, ParseParams, WaveInstances);
 }
 
 void USoundSimple::ChooseSoundWave()
 {
+	if (!Variations.Num())
+	{
+		return;
+	}
+
 	float ProbabilitySum = 0.0f;
 	for (int32 i = 0; i < Variations.Num(); ++i)
 	{
